PSD Importer: Promoted PSD category to the top
- Renamed categories to PSD.
- Added customization to move them further up the details panel.

diff --git a/Engine/Plugins/Experimental/PSDImporter/Source/PSDImporter/Private/PSDQuadActor.cpp b/Engine/Plugins/Experimental/PSDImporter/Source/PSDImporter/Private/PSDQuadActor.cpp
--- a/Engine/Plugins/Experimental/PSDImporter/Source/PSDImporter/Private/PSDQuadActor.cpp
+++ b/Engine/Plugins/Experimental/PSDImporter/Source/PSDImporter/Private/PSDQuadActor.cpp
@@ -134,7 +134,7 @@ void APSDQuadActor::InitComplete()
 
 FString APSDQuadActor::GetDefaultActorLabel() const
 {
-	return TEXT("Layer Root Actor");
+	return TEXT("PSD Layer Root Actor");
 }
 #endif
 
diff --git a/Engine/Plugins/Experimental/PSDImporter/Source/PSDImporter/Public/PSDQuadActor.h b/Engine/Plugins/Experimental/PSDImporter/Source/PSDImporter/Public/PSDQuadActor.h
--- a/Engine/Plugins/Experimental/PSDImporter/Source/PSDImporter/Public/PSDQuadActor.h
+++ b/Engine/Plugins/Experimental/PSDImporter/Source/PSDImporter/Public/PSDQuadActor.h
@@ -13,7 +13,7 @@ class UPSDDocument;
 class USceneComponent;
 class APSDQuadMeshActor;
 
-UCLASS(MinimalAPI, DisplayName = "Layer Root Actor")
+UCLASS(MinimalAPI, DisplayName = "PSD Layer Root Actor")
 class APSDQuadActor : public AActor
 {
 	GENERATED_BODY()
@@ -70,19 +70,19 @@ private:
 	UPROPERTY()
 	TArray<TWeakObjectPtr<APSDQuadMeshActor>> MeshListWeak;
 
-	UPROPERTY(VisibleAnywhere, Category = "Layer Root")
+	UPROPERTY(VisibleAnywhere, Category = "PSD")
 	TObjectPtr<USceneComponent> LayerRoot;
 
 	/** WIll separate each layer by this amount. */
-	UPROPERTY(EditInstanceOnly, Setter, Category = "Layer Root")
+	UPROPERTY(EditInstanceOnly, Setter, Category = "PSD")
 	float LayerDepthOffset = 1.f;
 
 	/** Will reduce the size of nearer quads to account for view distance with respect to the layer separation distance. 0 to disable. */
-	UPROPERTY(EditInstanceOnly, Setter, Category = "Layer Root")
+	UPROPERTY(EditInstanceOnly, Setter, Category = "PSD")
 	float AdjustForViewDistance = 0.f;
 
 	/** When assigning sort priority, use this as the first layer's priority. 0 to disable setting sort priority. */
-	UPROPERTY(EditInstanceOnly, Category = "Layer Root")
+	UPROPERTY(EditInstanceOnly, Category = "PSD")
 	int32 BaseTranslucentSortPriority = 1;
 
 	void OnLayerDepthOffsetChanged();
diff --git a/Engine/Plugins/Experimental/PSDImporter/Source/PSDImporter/Public/PSDQuadMeshActor.h b/Engine/Plugins/Experimental/PSDImporter/Source/PSDImporter/Public/PSDQuadMeshActor.h
--- a/Engine/Plugins/Experimental/PSDImporter/Source/PSDImporter/Public/PSDQuadMeshActor.h
+++ b/Engine/Plugins/Experimental/PSDImporter/Source/PSDImporter/Public/PSDQuadMeshActor.h
@@ -45,7 +45,7 @@ ENUM_CLASS_FLAGS(EPSDImporterLayerMaterialType)
  */
 DECLARE_MULTICAST_DELEGATE_OneParam(FPSDImporterTextureResetDelegate, APSDQuadMeshActor& /* Actor on which to reset the reset */)
 
-UCLASS(MinimalAPI, DisplayName = "Layer Actor")
+UCLASS(MinimalAPI, DisplayName = "PSD Layer Actor")
 class APSDQuadMeshActor : public AActor
 {
 	GENERATED_BODY()
@@ -71,22 +71,22 @@ public:
 	PSDIMPORTER_API void InitLayer(APSDQuadActor& InQuadActor, int32 InLayerIndex, UMaterialInterface* InLayerMaterial);
 #endif
 
-	UFUNCTION(CallInEditor, Category = "Layer", DisplayName = "Reset All")
+	UFUNCTION(CallInEditor, Category = "PSD", DisplayName = "Reset All")
 	PSDIMPORTER_API void ResetQuad();
 
-	UFUNCTION(CallInEditor, Category = "Layer", DisplayName = "Reset Depth")
+	UFUNCTION(CallInEditor, Category = "PSD", DisplayName = "Reset Depth")
 	PSDIMPORTER_API void ResetQuadDepth();
 
-	UFUNCTION(CallInEditor, Category = "Layer", DisplayName = "Reset Position")
+	UFUNCTION(CallInEditor, Category = "PSD", DisplayName = "Reset Position")
 	PSDIMPORTER_API void ResetQuadPosition();
 
-	UFUNCTION(CallInEditor, Category = "Layer", DisplayName = "Reset Size")
+	UFUNCTION(CallInEditor, Category = "PSD", DisplayName = "Reset Size")
 	PSDIMPORTER_API void ResetQuadSize();
 
-	UFUNCTION(CallInEditor, Category = "Layer", DisplayName = "Reset Texture")
+	UFUNCTION(CallInEditor, Category = "PSD", DisplayName = "Reset Texture")
 	PSDIMPORTER_API void ResetQuadTexture();
 
-	UFUNCTION(CallInEditor, Category = "Layer", DisplayName = "Reset Translucent Sort Priority")
+	UFUNCTION(CallInEditor, Category = "PSD", DisplayName = "Reset Translucent Sort Priority")
 	PSDIMPORTER_API void ResetQuadTranslucentSortPriority();
 
 #if WITH_EDITOR
@@ -98,12 +98,12 @@ public:
 private:
 	PSDIMPORTER_API static FPSDImporterTextureResetDelegate TextureResetDelegate;
 
-	UPROPERTY(VisibleAnywhere, DisplayName = "Layer Root Actor", Category = "Layer")
+	UPROPERTY(VisibleAnywhere, DisplayName = "Layer Root Actor", Category = "PSD")
 	TWeakObjectPtr<APSDQuadActor> QuadActorWeak;
 
-	UPROPERTY(VisibleAnywhere, Category = "Layer")
+	UPROPERTY(VisibleAnywhere, Category = "PSD")
 	int32 LayerIndex = INDEX_NONE;
 
-	UPROPERTY(VisibleAnywhere, Category = "Layer")
+	UPROPERTY(VisibleAnywhere, Category = "PSD")
 	TObjectPtr<UStaticMeshComponent> Mesh;
 };
\ No newline at end of file
diff --git a/Engine/Plugins/Experimental/PSDImporter/Source/PSDImporterEditor/Private/Customizations/PSDActorCustomization.cpp b/Engine/Plugins/Experimental/PSDImporter/Source/PSDImporterEditor/Private/Customizations/PSDActorCustomization.cpp
new file mode 100644
--- /dev/null
+++ b/Engine/Plugins/Experimental/PSDImporter/Source/PSDImporterEditor/Private/Customizations/PSDActorCustomization.cpp
@@ -0,0 +1,26 @@
+// Copyright Epic Games, Inc. All Rights Reserved.
+
+#include "Customizations/PSDActorCustomization.h"
+
+#include "DetailCategoryBuilder.h"
+#include "DetailLayoutBuilder.h"
+
+#define LOCTEXT_NAMESPACE "PSDActorCustomization"
+
+namespace UE::PSDImporterEditor
+{
+
+TSharedRef<IDetailCustomization> FPSDActorCustomization::MakeInstance()
+{
+	return MakeShared<FPSDActorCustomization>();
+}
+
+void FPSDActorCustomization::CustomizeDetails(IDetailLayoutBuilder& InDetailBuilder)
+{
+	InDetailBuilder.EditCategory(TEXT("PSD"), LOCTEXT("PSD", "PSD"))
+		.SetSortOrder(2000);
+}
+
+} // UE::PSDImporterEditor
+
+#undef LOCTEXT_NAMESPACE
diff --git a/Engine/Plugins/Experimental/PSDImporter/Source/PSDImporterEditor/Private/Customizations/PSDActorCustomization.h b/Engine/Plugins/Experimental/PSDImporter/Source/PSDImporterEditor/Private/Customizations/PSDActorCustomization.h
new file mode 100644
--- /dev/null
+++ b/Engine/Plugins/Experimental/PSDImporter/Source/PSDImporterEditor/Private/Customizations/PSDActorCustomization.h
@@ -0,0 +1,22 @@
+// Copyright Epic Games, Inc. All Rights Reserved.
+
+#pragma once
+
+#include "IDetailCustomization.h"
+
+namespace UE::PSDImporterEditor
+{
+
+/** Sorts the PSD category to the top. */
+class FPSDActorCustomization : public IDetailCustomization
+{
+public:
+	static TSharedRef<IDetailCustomization> MakeInstance();
+
+protected:
+	//~ Begin IDetailCustomization
+	virtual void CustomizeDetails(IDetailLayoutBuilder& InDetailBuilder) override;
+	//~ End IDetailCustomization
+};
+
+} // UE::PSDImporterEditor
diff --git a/Engine/Plugins/Experimental/PSDImporter/Source/PSDImporterEditor/Private/PSDImporterEditorModule.cpp b/Engine/Plugins/Experimental/PSDImporter/Source/PSDImporterEditor/Private/PSDImporterEditorModule.cpp
--- a/Engine/Plugins/Experimental/PSDImporter/Source/PSDImporterEditor/Private/PSDImporterEditorModule.cpp
+++ b/Engine/Plugins/Experimental/PSDImporter/Source/PSDImporterEditor/Private/PSDImporterEditorModule.cpp
@@ -2,6 +2,7 @@
 
 #include "Modules/ModuleInterface.h"
 
+#include "Customizations/PSDActorCustomization.h"
 #include "Customizations/PSDLayerCustomization.h"
 #include "Delegates/IDelegateInstance.h"
 #include "HAL/IConsoleManager.h"
@@ -29,13 +30,14 @@ public:
 	{
 		FPropertyEditorModule& PropertyModule = FModuleManager::LoadModuleChecked<FPropertyEditorModule>("PropertyEditor");
 		PropertyModule.RegisterCustomPropertyTypeLayout(FPSDFileLayer::StaticStruct()->GetFName(), FOnGetPropertyTypeCustomizationInstance::CreateStatic(&UE::PSDImporterEditor::FPSDLayerCustomization::MakeInstance));
+		PropertyModule.RegisterCustomClassLayout(APSDQuadActor::StaticClass()->GetFName(), FOnGetDetailCustomizationInstance::CreateStatic(&UE::PSDImporterEditor::FPSDActorCustomization::MakeInstance));
+		PropertyModule.RegisterCustomClassLayout(APSDQuadMeshActor::StaticClass()->GetFName(), FOnGetDetailCustomizationInstance::CreateStatic(&UE::PSDImporterEditor::FPSDActorCustomization::MakeInstance));
 
-		TSharedRef<FPropertySection> RootLayersSection = PropertyModule.FindOrCreateSection(APSDQuadActor::StaticClass()->GetFName(), TEXT("Layers"), LOCTEXT("Layers", "Layers"));
-		RootLayersSection->AddCategory(TEXT("Layer Root"));
+		TSharedRef<FPropertySection> RootLayersSection = PropertyModule.FindOrCreateSection(APSDQuadActor::StaticClass()->GetFName(), TEXT("PSD"), LOCTEXT("PSD", "PSD"));
+		RootLayersSection->AddCategory(TEXT("PSD"));
 
-		TSharedRef<FPropertySection> MeshLayersSection = PropertyModule.FindOrCreateSection(APSDQuadMeshActor::StaticClass()->GetFName(), TEXT("Layer"), LOCTEXT("Layer", "Layer"));
-		MeshLayersSection->AddCategory(TEXT("Layer"));
-		MeshLayersSection->AddCategory(TEXT("Materials"));
+		TSharedRef<FPropertySection> MeshLayersSection = PropertyModule.FindOrCreateSection(APSDQuadMeshActor::StaticClass()->GetFName(), TEXT("PSD"), LOCTEXT("PSD", "PSD"));
+		MeshLayersSection->AddCategory(TEXT("PSD"));
 
 		FPSDImporterContentBrowserIntegration::Get().Integrate();
 
