PSD Importer: No longer multi-imports files

diff --git a/Engine/Plugins/Experimental/PSDImporter/Source/PSDImporterEditor/Private/Factories/PSDDocumentImportFactory_Visitors.cpp b/Engine/Plugins/Experimental/PSDImporter/Source/PSDImporterEditor/Private/Factories/PSDDocumentImportFactory_Visitors.cpp
--- a/Engine/Plugins/Experimental/PSDImporter/Source/PSDImporterEditor/Private/Factories/PSDDocumentImportFactory_Visitors.cpp
+++ b/Engine/Plugins/Experimental/PSDImporter/Source/PSDImporterEditor/Private/Factories/PSDDocumentImportFactory_Visitors.cpp
@@ -15,6 +15,8 @@
 #include "PSDImporterEditorUtilities.h"
 #include "PSDLayerTextureUserData.h"
 #include "Subsystems/EditorAssetSubsystem.h"
+#include "UObject/Object.h"
+#include "UObject/SoftObjectPtr.h"
 
 FPSDDocumentImportFactory_Visitors::FPSDDocumentImportFactory_Visitors(const FString& InFilePath, UPSDDocument* InDocument)
 	: FilePath(InFilePath)
@@ -315,10 +317,26 @@ UObject* FPSDDocumentImportFactory_Visitors::MakeAsset(UClass* InClass, const FS
 	FString AssetName;
 	MakeAssetPath(InDocumentName, InAssetName, InAssetPrefix, BasePath, AssetName);
 
+	const FString AssetPath = BasePath / AssetName + TEXT(".") + AssetName;
+	const bool bAlreadyImported = ImportedAssets.Contains(AssetPath);
+	ImportedAssets.Add(AssetPath);
+
+	// If not already imported by us, the CreateAsset call will ask the user if they want to overwrite 
+	// an existing asset.
+	if (bAlreadyImported)
+	{
+		TSoftObjectPtr<UObject> ExistingObjectPtr = TSoftObjectPtr<UObject>(FSoftObjectPath(*AssetPath));
+
+		if (UObject* ExistingObject = ExistingObjectPtr.LoadSynchronous())
+		{
+			return ExistingObject;
+		}
+	}
+
 	IAssetTools& AssetTools = FAssetToolsModule::GetModule().Get();
+
 	FString PackageName;
-	AssetTools.CreateUniqueAssetName(BasePath / AssetName, TEXT(""), PackageName, AssetName);
-	return AssetTools.CreateAsset(AssetName, FPaths::GetPath(PackageName), InClass, nullptr);
+	return AssetTools.CreateAsset(AssetName, BasePath, InClass, nullptr);
 }
 
 void FPSDDocumentImportFactory_Visitors::MakeAssetPath(
diff --git a/Engine/Plugins/Experimental/PSDImporter/Source/PSDImporterEditor/Private/Factories/PSDDocumentImportFactory_Visitors.h b/Engine/Plugins/Experimental/PSDImporter/Source/PSDImporterEditor/Private/Factories/PSDDocumentImportFactory_Visitors.h
--- a/Engine/Plugins/Experimental/PSDImporter/Source/PSDImporterEditor/Private/Factories/PSDDocumentImportFactory_Visitors.h
+++ b/Engine/Plugins/Experimental/PSDImporter/Source/PSDImporterEditor/Private/Factories/PSDDocumentImportFactory_Visitors.h
@@ -4,8 +4,9 @@
 
 #include "PSDFileImport.h"
 
-#include "PSDFile.h"
+#include "Containers/Set.h"
 #include "Containers/UnrealString.h"
+#include "PSDFile.h"
 #include "UObject/StrongObjectPtr.h"
 
 class UPSDDocument;
@@ -34,4 +35,5 @@ private:
 	FPSDFileDocument& FileDocument;
 	TSet<FPSDFileLayer> OldLayers;
 	TSet<FPSDFileLayer> NewLayers;
+	TSet<FString> ImportedAssets;
 };
