PSD Importer: Resetting camera offset distance now correctly resets layer positions.

diff --git a/Engine/Plugins/Experimental/PSDImporter/Source/PSDImporter/Private/PSDQuadActor.cpp b/Engine/Plugins/Experimental/PSDImporter/Source/PSDImporter/Private/PSDQuadActor.cpp
--- a/Engine/Plugins/Experimental/PSDImporter/Source/PSDImporter/Private/PSDQuadActor.cpp
+++ b/Engine/Plugins/Experimental/PSDImporter/Source/PSDImporter/Private/PSDQuadActor.cpp
@@ -154,6 +154,7 @@ void APSDQuadActor::UpdateQuadSeparationDistances()
 
 void APSDQuadActor::OnAdjustForViewDistanceChanged()
 {
+	UpdateQuadSeparationDistances();
 	UpdateQuadSizeForViewDistance();
 }
 
@@ -161,8 +162,8 @@ void APSDQuadActor::UpdateQuadSizeForViewDistance()
 {
 	for (APSDQuadMeshActor* MeshActor : GetQuadMeshes())
 	{
-		MeshActor->ResetQuadPosition();
 		MeshActor->ResetQuadSize();
+		MeshActor->ResetQuadPosition();
 	}
 }
 
