PSD Importer: Importing layers with special characters now sanitzes the layer name.

diff --git a/Engine/Plugins/Experimental/PSDImporter/Source/PSDImporterEditor/Private/Factories/PSDDocumentImportFactory_Visitors.cpp b/Engine/Plugins/Experimental/PSDImporter/Source/PSDImporterEditor/Private/Factories/PSDDocumentImportFactory_Visitors.cpp
--- a/Engine/Plugins/Experimental/PSDImporter/Source/PSDImporterEditor/Private/Factories/PSDDocumentImportFactory_Visitors.cpp
+++ b/Engine/Plugins/Experimental/PSDImporter/Source/PSDImporterEditor/Private/Factories/PSDDocumentImportFactory_Visitors.cpp
@@ -18,6 +18,38 @@
 #include "UObject/Object.h"
 #include "UObject/SoftObjectPtr.h"
 
+namespace UE::PSDImporter::Private
+{
+	void Sanitize(FString& OutPath, const FString& InInvalidChars, const FString::ElementType InReplaceWith)
+	{
+		// Make sure the package path only contains valid characters
+		FText OutError;
+		
+		if (FName::IsValidXName(OutPath, InInvalidChars, &OutError))
+		{
+			return;
+		}
+
+		for (const FString::ElementType& InvalidChar : InInvalidChars)
+		{
+			OutPath.ReplaceCharInline(InvalidChar, InReplaceWith, ESearchCase::CaseSensitive);
+		}
+	}
+
+	// Replaces any invalid package path characters 
+	void SanitizePackagePath(FString& OutPath, FString::ElementType InReplaceWith = TEXT('_'))
+	{
+		FPaths::RemoveDuplicateSlashes(OutPath);
+		return Private::Sanitize(OutPath, INVALID_OBJECTPATH_CHARACTERS INVALID_LONGPACKAGE_CHARACTERS, InReplaceWith);
+	}
+
+	// Replaces any invalid object name characters 
+	void SanitizeAssetName(FString& OutPath, FString::ElementType InReplaceWith = TEXT('_'))
+	{
+		return Private::Sanitize(OutPath, INVALID_OBJECTNAME_CHARACTERS, InReplaceWith);
+	}
+}
+
 FPSDDocumentImportFactory_Visitors::FPSDDocumentImportFactory_Visitors(const FString& InFilePath, UPSDDocument* InDocument)
 	: FilePath(InFilePath)
     , Document(InDocument)
@@ -313,11 +345,16 @@ UObject* FPSDDocumentImportFactory_Visitors::MakeAsset(UClass* InClass, const FS
 		return nullptr;
 	}
 
+	using namespace UE::PSDImporter::Private;
+
 	FString BasePath;
 	FString AssetName;
 	MakeAssetPath(InDocumentName, InAssetName, InAssetPrefix, BasePath, AssetName);
 
-	const FString AssetPath = BasePath / AssetName + TEXT(".") + AssetName;
+	SanitizePackagePath(BasePath);
+	SanitizeAssetName(AssetName);
+
+	FString AssetPath = BasePath / AssetName + TEXT(".") + AssetName;
 	const bool bAlreadyImported = ImportedAssets.Contains(AssetPath);
 	ImportedAssets.Add(AssetPath);
 
