PSD Importer: Hidden Quad Actor property AdjustForViewDistance because it is not user friendly.

diff --git a/Engine/Plugins/Experimental/PSDImporter/Source/PSDImporter/Public/PSDQuadActor.h b/Engine/Plugins/Experimental/PSDImporter/Source/PSDImporter/Public/PSDQuadActor.h
--- a/Engine/Plugins/Experimental/PSDImporter/Source/PSDImporter/Public/PSDQuadActor.h
+++ b/Engine/Plugins/Experimental/PSDImporter/Source/PSDImporter/Public/PSDQuadActor.h
@@ -78,7 +78,7 @@ private:
 	float LayerDepthOffset = 1.f;
 
 	/** Will reduce the size of nearer quads to account for view distance with respect to the layer separation distance. 0 to disable. */
-	UPROPERTY(EditInstanceOnly, Setter, Category = "PSD")
+	UPROPERTY()
 	float AdjustForViewDistance = 0.f;
 
 	/** When assigning sort priority, use this as the first layer's priority. 0 to disable setting sort priority. */
