PSD Importer: Importing 16 and 32-bit PSDs now works correctly.

diff --git a/Engine/Plugins/Experimental/PSDImporter/Source/PSDImporterCore/Private/Readers/LayerReader.cpp b/Engine/Plugins/Experimental/PSDImporter/Source/PSDImporterCore/Private/Readers/LayerReader.cpp
--- a/Engine/Plugins/Experimental/PSDImporter/Source/PSDImporterCore/Private/Readers/LayerReader.cpp
+++ b/Engine/Plugins/Experimental/PSDImporter/Source/PSDImporterCore/Private/Readers/LayerReader.cpp
@@ -64,7 +64,16 @@ namespace UE::PSDImporter::Private
 		return EPSDBlendMode::Unknown;
 	}
 
+	template<typename InScanlineType>
+	InScanlineType GetPixelFromScanline(const InScanlineType* InInputScanline, const uint32 InPixel)
+	{
+		if (InInputScanline)
+		{
+			return InInputScanline[InPixel];
+		}
 
+		return 0;
+	};
 
 	template <typename DataType, int Depth>
 	TFuture<FImage> ReadRGBAInternal(FReadContext& InContext, psd::Layer* InLayer)
@@ -124,20 +133,22 @@ namespace UE::PSDImporter::Private
 			}
 
 			FImage OutputImage;
-			OutputImage.Init(Width, Height, ERawImageFormat::BGRA8);
 
-			// Copy planes to (interleaved) image
+			if constexpr (Depth == 8)
 			{
-				auto GetPixelFromScanline = []<typename InScanlineType>(const InScanlineType* InInputScanline, const uint32 InPixel) -> InScanlineType
-					{
-						if (InInputScanline)
-						{
-							return InInputScanline[InPixel];
-						}
-
-						return 0;
-					};
+				OutputImage.Init(Width, Height, ERawImageFormat::BGRA8);
+			}
+			else if constexpr (Depth == 16)
+			{
+				OutputImage.Init(Width, Height, ERawImageFormat::RGBA16);
+			}
+			else if constexpr (Depth == 32)
+			{
+				OutputImage.Init(Width, Height, ERawImageFormat::RGBA32F);
+			}
 
+			// Copy planes to (interleaved) image
+			{
 				// Max 32 bits
 				constexpr uint8 One[4] = {255, 255, 255, 255};
 				
@@ -148,22 +159,25 @@ namespace UE::PSDImporter::Private
 
 				for (uint32 RowIdx = 0; RowIdx < Height; ++RowIdx, OutputScanline += OutputScanlineSize)
 				{
-					const uint8* InputScanline[4] = {};
-					uint32 AlphaMask = ~0U;
+					const uint32 OutputScanLineEnd = RowIdx * OutputScanlineSize;
 
-					for (uint16 ChannelIdx = 0; ChannelIdx < NumChannels; ++ChannelIdx)
+					if (OutputScanLineEnd >= OutputImage.RawData.Num())
 					{
-						InputScanline[ChannelIdx] = reinterpret_cast<const uint8*>(CanvasData[ChannelIdx]) + (RowIdx * InputScanlineSizePerChannel);
+						break;
 					}
 
+					const uint8* InputScanline[4] = {};
+					uint32 AlphaMask = -1;
+
 					if (NumChannels == 3)
 					{
 						InputScanline[3] = One;
 						AlphaMask = 0;
 					}
-					else if (NumChannels == 4)
+
+					for (uint16 ChannelIdx = 0; ChannelIdx < NumChannels; ++ChannelIdx)
 					{
-						AlphaMask = ~0U;
+						InputScanline[ChannelIdx] = reinterpret_cast<const uint8*>(CanvasData[ChannelIdx]) + (RowIdx * InputScanlineSizePerChannel);
 					}
 
 					if constexpr (Depth == 8)
@@ -172,10 +186,10 @@ namespace UE::PSDImporter::Private
 
 						for (uint32 PixelX = 0; PixelX < Width; ++PixelX)
 						{
-							Scanline8[PixelX].R = GetPixelFromScanline(InputScanline[0], PixelX);
-							Scanline8[PixelX].G = GetPixelFromScanline(InputScanline[1], PixelX);
-							Scanline8[PixelX].B = GetPixelFromScanline(InputScanline[2], PixelX);
-							Scanline8[PixelX].A = GetPixelFromScanline(InputScanline[3], PixelX & AlphaMask);
+							Scanline8[PixelX].R = GetPixelFromScanline<uint8>(InputScanline[0], PixelX);
+							Scanline8[PixelX].G = GetPixelFromScanline<uint8>(InputScanline[1], PixelX);
+							Scanline8[PixelX].B = GetPixelFromScanline<uint8>(InputScanline[2], PixelX);
+							Scanline8[PixelX].A = GetPixelFromScanline<uint8>(InputScanline[3], PixelX) & AlphaMask;
 						}
 					}
 					else if constexpr (Depth == 16)
@@ -185,14 +199,22 @@ namespace UE::PSDImporter::Private
 						const uint16* InputScanlineB16 = reinterpret_cast<const uint16*>(InputScanline[2]);
 						const uint16* InputScanlineA16 = reinterpret_cast<const uint16*>(InputScanline[3]);
 
-						uint16* Scanline16 = reinterpret_cast<uint16*>(OutputScanline);
+						struct FColorUInt16
+						{
+							uint16 R;
+							uint16 G;
+							uint16 B;
+							uint16 A;
+						};
+
+						FColorUInt16* Scanline16 = reinterpret_cast<FColorUInt16*>(OutputScanline);
 
 						for (uint32 PixelX = 0; PixelX < Width; ++PixelX)
 						{
-							Scanline16[4 * PixelX + 0] = GetPixelFromScanline(InputScanlineR16, PixelX);
-							Scanline16[4 * PixelX + 1] = GetPixelFromScanline(InputScanlineG16, PixelX);
-							Scanline16[4 * PixelX + 2] = GetPixelFromScanline(InputScanlineB16, PixelX);
-							Scanline16[4 * PixelX + 3] = GetPixelFromScanline(InputScanlineA16, PixelX & AlphaMask);
+							Scanline16[PixelX].R = GetPixelFromScanline<uint16>(InputScanlineR16, PixelX);
+							Scanline16[PixelX].G = GetPixelFromScanline<uint16>(InputScanlineG16, PixelX);
+							Scanline16[PixelX].B = GetPixelFromScanline<uint16>(InputScanlineB16, PixelX);
+							Scanline16[PixelX].A = GetPixelFromScanline<uint16>(InputScanlineA16, PixelX) & AlphaMask;
 						}
 					}
 					else if constexpr (Depth == 32)
@@ -202,14 +224,23 @@ namespace UE::PSDImporter::Private
 						const uint32* InputScanlineB32 = reinterpret_cast<const uint32*>(InputScanline[2]);
 						const uint32* InputScanlineA32 = reinterpret_cast<const uint32*>(InputScanline[3]);
 
-						uint32* Scanline32 = reinterpret_cast<uint32*>(OutputScanline);
+						struct FColorF32
+						{
+							float R;
+							float G;
+							float B;
+							float A;
+						};
+						
+						FColorF32* Scanline32 = reinterpret_cast<FColorF32*>(OutputScanline);
+						constexpr float MaxValue = 4294967296.f;
 
 						for (uint32 PixelX = 0; PixelX < Width; ++PixelX)
 						{
-							Scanline32[8 * PixelX + 0] = GetPixelFromScanline(InputScanlineR32, PixelX);
-							Scanline32[8 * PixelX + 1] = GetPixelFromScanline(InputScanlineG32, PixelX);
-							Scanline32[8 * PixelX + 2] = GetPixelFromScanline(InputScanlineB32, PixelX);
-							Scanline32[8 * PixelX + 3] = GetPixelFromScanline(InputScanlineA32, PixelX & AlphaMask);
+							Scanline32[PixelX].R = static_cast<float>(GetPixelFromScanline<uint32>(InputScanlineR32, PixelX)) / MaxValue;
+							Scanline32[PixelX].G = static_cast<float>(GetPixelFromScanline<uint32>(InputScanlineG32, PixelX)) / MaxValue;
+							Scanline32[PixelX].B = static_cast<float>(GetPixelFromScanline<uint32>(InputScanlineB32, PixelX)) / MaxValue;
+							Scanline32[PixelX].A = static_cast<float>(GetPixelFromScanline<uint32>(InputScanlineA32, PixelX) & AlphaMask) / MaxValue;
 						}
 					}
 				}
